Move the source for the GetBuildDate() function into the BuildSettings module, so it's always updated whenever the Build.version file changes.

diff --git a/Engine/Source/Runtime/BuildSettings/Private/BuildSettings.cpp b/Engine/Source/Runtime/BuildSettings/Private/BuildSettings.cpp
--- a/Engine/Source/Runtime/BuildSettings/Private/BuildSettings.cpp
+++ b/Engine/Source/Runtime/BuildSettings/Private/BuildSettings.cpp
@@ -23,6 +23,11 @@ namespace BuildSettings
 	{
 		return TEXT(BRANCH_NAME);
 	}
+	
+	const TCHAR* GetBuildDate()
+	{
+		return TEXT(__DATE__);
+	}
 
 	const TCHAR* GetBuildVersion()
 	{
diff --git a/Engine/Source/Runtime/BuildSettings/Public/BuildSettings.h b/Engine/Source/Runtime/BuildSettings/Public/BuildSettings.h
--- a/Engine/Source/Runtime/BuildSettings/Public/BuildSettings.h
+++ b/Engine/Source/Runtime/BuildSettings/Public/BuildSettings.h
@@ -42,6 +42,13 @@ namespace BuildSettings
 	 */
 	BUILDSETTINGS_API const TCHAR* GetBranchName();
 
+	/**
+	 * The timestamp of this build. Derived from the compiler's __DATE__ macro, so only updated when the BuildSettings module is rebuilt (ie. whenever CL changes, etc...)
+	 *
+	 * @return String representing the build timestamp
+	 */
+	BUILDSETTINGS_API const TCHAR* GetBuildDate();
+
 	/**
 	 * Retrieves the user-defined build version for this application.
 	 *
diff --git a/Engine/Source/Runtime/Core/Private/Misc/App.cpp b/Engine/Source/Runtime/Core/Private/Misc/App.cpp
--- a/Engine/Source/Runtime/Core/Private/Misc/App.cpp
+++ b/Engine/Source/Runtime/Core/Private/Misc/App.cpp
@@ -95,7 +95,7 @@ void FApp::SetDebugGame(bool bInIsDebugGame)
 
 FString FApp::GetBuildDate()
 {
-	return FString(ANSI_TO_TCHAR(__DATE__));
+	return FString(BuildSettings::GetBuildDate());
 }
 
 
