Add GetBuildTime() to BuildSettings

#rnx

diff --git a/Engine/Source/Runtime/BuildSettings/Private/BuildSettings.cpp b/Engine/Source/Runtime/BuildSettings/Private/BuildSettings.cpp
--- a/Engine/Source/Runtime/BuildSettings/Private/BuildSettings.cpp
+++ b/Engine/Source/Runtime/BuildSettings/Private/BuildSettings.cpp
@@ -29,6 +29,11 @@ namespace BuildSettings
 		return TEXT(__DATE__);
 	}
 
+	const TCHAR* GetBuildTime()
+	{
+		return TEXT(__TIME__);
+	}
+
 	const TCHAR* GetBuildVersion()
 	{
 		return TEXT(BUILD_VERSION);
diff --git a/Engine/Source/Runtime/BuildSettings/Public/BuildSettings.h b/Engine/Source/Runtime/BuildSettings/Public/BuildSettings.h
--- a/Engine/Source/Runtime/BuildSettings/Public/BuildSettings.h
+++ b/Engine/Source/Runtime/BuildSettings/Public/BuildSettings.h
@@ -43,12 +43,19 @@ namespace BuildSettings
 	BUILDSETTINGS_API const TCHAR* GetBranchName();
 
 	/**
-	 * The timestamp of this build. Derived from the compiler's __DATE__ macro, so only updated when the BuildSettings module is rebuilt (ie. whenever CL changes, etc...)
+	 * The date timestamp of this build. Derived from the compiler's __DATE__ macro, so only updated when the BuildSettings module is rebuilt (ie. whenever CL changes, etc...)
 	 *
 	 * @return String representing the build timestamp
 	 */
 	BUILDSETTINGS_API const TCHAR* GetBuildDate();
 
+	/**
+	 * The time timestamp of this build. Derived from the compiler's __TIME__ macro, so only updated when the BuildSettings module is rebuilt (ie. whenever CL changes, etc...)
+	 *
+	 * @return String representing the build timestamp
+	 */
+	BUILDSETTINGS_API const TCHAR* GetBuildTime();
+
 	/**
 	 * Retrieves the user-defined build version for this application.
 	 *
