Allow project targets to override current CL settings.

diff --git a/Engine/Source/Runtime/BuildSettings/BuildSettings.Build.cs b/Engine/Source/Runtime/BuildSettings/BuildSettings.Build.cs
--- a/Engine/Source/Runtime/BuildSettings/BuildSettings.Build.cs
+++ b/Engine/Source/Runtime/BuildSettings/BuildSettings.Build.cs
@@ -2,6 +2,7 @@
 
 using UnrealBuildTool;
 using System.IO;
+using System;
 
 public class BuildSettings : ModuleRules
 {
@@ -17,7 +18,19 @@ public class BuildSettings : ModuleRules
 		PrivateDefinitions.Add(string.Format("ENGINE_IS_LICENSEE_VERSION={0}", Target.Version.IsLicenseeVersion? "true" : "false"));
 		PrivateDefinitions.Add(string.Format("ENGINE_IS_PROMOTED_BUILD={0}", Target.Version.IsPromotedBuild? "true" : "false"));
 
-		PrivateDefinitions.Add(string.Format("CURRENT_CHANGELIST={0}", Target.Version.Changelist));
+		bool bContainsCurrentChangelist = false;
+		foreach (string item in Target.GlobalDefinitions)
+		{
+			if (item.Contains("CURRENT_CHANGELIST"))
+			{
+				bContainsCurrentChangelist = true;
+				break;
+			}
+		}
+		if (!bContainsCurrentChangelist)
+		{
+			PrivateDefinitions.Add(string.Format("CURRENT_CHANGELIST={0}", Target.Version.Changelist));
+		}
 		PrivateDefinitions.Add(string.Format("COMPATIBLE_CHANGELIST={0}", Target.Version.EffectiveCompatibleChangelist));
 
 		PrivateDefinitions.Add(string.Format("BRANCH_NAME=\"{0}\"", Target.Version.BranchName));
