BuildSettings: Add functions to get engine version & clean up ModuleRules

* use .Any() instead of looping through GlobalDefinitions
* BUILD_SOURCE_URL can always be set as interpolated string replace null with an empty string

#rnx

diff --git a/Engine/Source/Runtime/BuildSettings/BuildSettings.Build.cs b/Engine/Source/Runtime/BuildSettings/BuildSettings.Build.cs
--- a/Engine/Source/Runtime/BuildSettings/BuildSettings.Build.cs
+++ b/Engine/Source/Runtime/BuildSettings/BuildSettings.Build.cs
@@ -1,47 +1,35 @@
 // Copyright Epic Games, Inc. All Rights Reserved.
 
 using UnrealBuildTool;
-using System.IO;
 using System;
+using System.Linq;
 
 public class BuildSettings : ModuleRules
 {
 	public BuildSettings(ReadOnlyTargetRules Target) : base(Target)
 	{
 		DeterministicWarningLevel = WarningLevel.Off; // This module intentionally uses __DATE__ and __TIME__ macros
+		UnsafeTypeCastWarningLevel = WarningLevel.Error;
+
 		PrivateIncludePathModuleNames.Add("Core");
 
 		bRequiresImplementModule = false;
 
-		PrivateDefinitions.Add(string.Format("ENGINE_VERSION_MAJOR={0}", Target.Version.MajorVersion));
-		PrivateDefinitions.Add(string.Format("ENGINE_VERSION_MINOR={0}", Target.Version.MinorVersion));
-		PrivateDefinitions.Add(string.Format("ENGINE_VERSION_HOTFIX={0}", Target.Version.PatchVersion));
-		PrivateDefinitions.Add(string.Format("ENGINE_IS_LICENSEE_VERSION={0}", Target.Version.IsLicenseeVersion? "true" : "false"));
-		PrivateDefinitions.Add(string.Format("ENGINE_IS_PROMOTED_BUILD={0}", Target.Version.IsPromotedBuild? "true" : "false"));
-
-		bool bContainsCurrentChangelist = false;
-		foreach (string item in Target.GlobalDefinitions)
-		{
-			if (item.Contains("CURRENT_CHANGELIST"))
-			{
-				bContainsCurrentChangelist = true;
-				break;
-			}
-		}
-		if (!bContainsCurrentChangelist)
+		PrivateDefinitions.Add($"ENGINE_VERSION_MAJOR={Target.Version.MajorVersion}");
+		PrivateDefinitions.Add($"ENGINE_VERSION_MINOR={Target.Version.MinorVersion}");
+		PrivateDefinitions.Add($"ENGINE_VERSION_HOTFIX={Target.Version.PatchVersion}");
+		PrivateDefinitions.Add($"ENGINE_VERSION_STRING=\"{Target.Version.MajorVersion}.{Target.Version.MinorVersion}.{Target.Version.PatchVersion}-{Target.BuildVersion}\"");
+		PrivateDefinitions.Add($"ENGINE_IS_LICENSEE_VERSION={(Target.Version.IsLicenseeVersion ? "true" : "false")}");
+		PrivateDefinitions.Add($"ENGINE_IS_PROMOTED_BUILD={(Target.Version.IsPromotedBuild ? "true" : "false")}");
+		
+		if (!Target.GlobalDefinitions.Any(x => x.Contains("CURRENT_CHANGELIST", StringComparison.Ordinal)))
 		{
-			PrivateDefinitions.Add(string.Format("CURRENT_CHANGELIST={0}", Target.Version.Changelist));
+			PrivateDefinitions.Add($"CURRENT_CHANGELIST={Target.Version.Changelist}");
 		}
-		PrivateDefinitions.Add(string.Format("COMPATIBLE_CHANGELIST={0}", Target.Version.EffectiveCompatibleChangelist));
 
-		PrivateDefinitions.Add(string.Format("BRANCH_NAME=\"{0}\"", Target.Version.BranchName));
-
-		PrivateDefinitions.Add(string.Format("BUILD_VERSION=\"{0}\"", Target.BuildVersion));
-		if (!String.IsNullOrEmpty(Target.Version.BuildURL))
-		{
-			PrivateDefinitions.Add(string.Format("BUILD_SOURCE_URL=\"{0}\"", Target.Version.BuildURL));
-		}
-				
-		UnsafeTypeCastWarningLevel = WarningLevel.Error;
+		PrivateDefinitions.Add($"COMPATIBLE_CHANGELIST={Target.Version.EffectiveCompatibleChangelist}");
+		PrivateDefinitions.Add($"BRANCH_NAME=\"{Target.Version.BranchName}\"");
+		PrivateDefinitions.Add($"BUILD_VERSION=\"{Target.BuildVersion}\"");
+		PrivateDefinitions.Add($"BUILD_SOURCE_URL=\"{Target.Version.BuildURL}\"");
 	}
 }
diff --git a/Engine/Source/Runtime/BuildSettings/Private/BuildSettings.cpp b/Engine/Source/Runtime/BuildSettings/Private/BuildSettings.cpp
--- a/Engine/Source/Runtime/BuildSettings/Private/BuildSettings.cpp
+++ b/Engine/Source/Runtime/BuildSettings/Private/BuildSettings.cpp
@@ -9,6 +9,26 @@ namespace BuildSettings
 		return ENGINE_IS_LICENSEE_VERSION;
 	}
 
+	int GetEngineVersionMajor()
+	{
+		return ENGINE_VERSION_MAJOR;
+	}
+
+	int GetEngineVersionMinor()
+	{
+		return ENGINE_VERSION_MINOR;
+	}
+
+	int GetEngineVersionHotfix()
+	{
+		return ENGINE_VERSION_HOTFIX;
+	}
+
+	const TCHAR* GetEngineVersionString()
+	{
+		return TEXT(ENGINE_VERSION_STRING);
+	}
+
 	int GetCurrentChangelist()
 	{
 		return CURRENT_CHANGELIST;
@@ -51,11 +71,7 @@ namespace BuildSettings
 	
 	const TCHAR* GetBuildURL()
 	{
-#ifdef BUILD_SOURCE_URL
 		return TEXT(BUILD_SOURCE_URL);
-#else
-		return TEXT("");
-#endif	
 	}
 }
 
diff --git a/Engine/Source/Runtime/BuildSettings/Public/BuildSettings.h b/Engine/Source/Runtime/BuildSettings/Public/BuildSettings.h
--- a/Engine/Source/Runtime/BuildSettings/Public/BuildSettings.h
+++ b/Engine/Source/Runtime/BuildSettings/Public/BuildSettings.h
@@ -16,6 +16,34 @@ namespace BuildSettings
 	 */
 	BUILDSETTINGS_API bool IsLicenseeVersion();
 
+	/**
+	* The major engine version.
+	* 
+	* @return The major part of the engine version
+	*/
+	BUILDSETTINGS_API int GetEngineVersionMajor();
+
+	/**
+	* The minor engine version.
+	*
+	* @return The minor part of the engine version
+	*/
+	BUILDSETTINGS_API int GetEngineVersionMinor();
+
+	/**
+	* The hotfix engine version.
+	*
+	* @return The hotfix patch part of the engine version
+	*/
+	BUILDSETTINGS_API int GetEngineVersionHotfix();
+
+	/**
+	* The engine version.
+	* 
+	* @return The engine version as a string in the format MAJOR.MINOR.HOTFIX-BUILD_VERSION
+	*/
+	BUILDSETTINGS_API const TCHAR* GetEngineVersionString();
+
 	/**
 	 * The Perforce changelist being compiled. Use this value advisedly; it does not take into account out-of-order commits to engine release branches over 
 	 * development branches, licensee versions, or whether the engine version has been locked to maintain compatibility with a previous engine release. Prefer
