[BuildSettings]
* Added GetBuildUser, GetBuildUserDomain and GetBuildMachine to BuildSettings. These functions will return private information about the build if Target.bEnablePrivateBuildInformation is set to true. If set to false all functions will return empty strings

diff --git a/Engine/Source/Programs/UnrealBuildTool/Configuration/ReadOnlyTargetRules.cs b/Engine/Source/Programs/UnrealBuildTool/Configuration/ReadOnlyTargetRules.cs
--- a/Engine/Source/Programs/UnrealBuildTool/Configuration/ReadOnlyTargetRules.cs
+++ b/Engine/Source/Programs/UnrealBuildTool/Configuration/ReadOnlyTargetRules.cs
@@ -219,6 +219,8 @@ namespace UnrealBuildTool
 
 		public bool bForceEnableRTTI => Inner.bForceEnableRTTI;
 
+		public bool bEnablePrivateBuildInformation => Inner.bEnablePrivateBuildInformation;
+
 		public bool bEnablePIE => Inner.bEnablePIE;
 
 		public bool bEnableStackProtection => Inner.bEnableStackProtection;
diff --git a/Engine/Source/Programs/UnrealBuildTool/Configuration/TargetRules.cs b/Engine/Source/Programs/UnrealBuildTool/Configuration/TargetRules.cs
--- a/Engine/Source/Programs/UnrealBuildTool/Configuration/TargetRules.cs
+++ b/Engine/Source/Programs/UnrealBuildTool/Configuration/TargetRules.cs
@@ -1249,6 +1249,12 @@ namespace UnrealBuildTool
 		[RequiresUniqueBuildEnvironment]
 		public bool bForceEnableRTTI { get; set; }
 
+		/// <summary>
+		/// Enable BuildSettings to contain private information about build. Such as machine name, user name and user domain name (exposed in BuildSettings.h)
+		/// </summary>
+		[RequiresUniqueBuildEnvironment]
+		public bool bEnablePrivateBuildInformation { get; set; }
+
 		/// <summary>
 		/// Enable Position Independent Executable (PIE). Has an overhead cost
 		/// </summary>
diff --git a/Engine/Source/Runtime/BuildSettings/BuildSettings.Build.cs b/Engine/Source/Runtime/BuildSettings/BuildSettings.Build.cs
--- a/Engine/Source/Runtime/BuildSettings/BuildSettings.Build.cs
+++ b/Engine/Source/Runtime/BuildSettings/BuildSettings.Build.cs
@@ -31,5 +31,18 @@ public class BuildSettings : ModuleRules
 		PrivateDefinitions.Add($"BRANCH_NAME=\"{Target.Version.BranchName}\"");
 		PrivateDefinitions.Add($"BUILD_VERSION=\"{Target.BuildVersion}\"");
 		PrivateDefinitions.Add($"BUILD_SOURCE_URL=\"{Target.Version.BuildURL}\"");
+
+		string userName = string.Empty;
+		string userDomainName = string.Empty;
+		string machineName = string.Empty;
+		if (Target.bEnablePrivateBuildInformation)
+		{
+			userName = Environment.UserName;
+			userDomainName = Environment.UserDomainName;
+			machineName = UnrealBuildBase.Unreal.MachineName;
+		}
+		PrivateDefinitions.Add($"BUILD_USER=\"{userName}\"");
+		PrivateDefinitions.Add($"BUILD_USERDOMAINNAME=\"{userDomainName}\"");
+		PrivateDefinitions.Add($"BUILD_MACHINENAME=\"{machineName}\"");
 	}
 }
diff --git a/Engine/Source/Runtime/BuildSettings/Private/BuildSettings.cpp b/Engine/Source/Runtime/BuildSettings/Private/BuildSettings.cpp
--- a/Engine/Source/Runtime/BuildSettings/Private/BuildSettings.cpp
+++ b/Engine/Source/Runtime/BuildSettings/Private/BuildSettings.cpp
@@ -73,6 +73,19 @@ namespace BuildSettings
 	{
 		return TEXT(BUILD_SOURCE_URL);
 	}
-}
 
+	const TCHAR* GetBuildUser()
+	{
+		return TEXT(BUILD_USER);
+	}
 
+	const TCHAR* GetBuildUserDomain()
+	{
+		return TEXT(BUILD_USERDOMAINNAME);
+	}
+
+	const TCHAR* GetBuildMachine()
+	{
+		return TEXT(BUILD_MACHINENAME);
+	}
+}
diff --git a/Engine/Source/Runtime/BuildSettings/Public/BuildSettings.h b/Engine/Source/Runtime/BuildSettings/Public/BuildSettings.h
--- a/Engine/Source/Runtime/BuildSettings/Public/BuildSettings.h
+++ b/Engine/Source/Runtime/BuildSettings/Public/BuildSettings.h
@@ -107,4 +107,19 @@ namespace BuildSettings
 	 * Returns a URL where the job which created these binaries on an automation system (e.g. Horde) can be found, or an empty string.
 	 */
 	BUILDSETTINGS_API const TCHAR* GetBuildURL();
+
+	/**
+	 * Returns the name of the user which created these binaries
+	 */
+	BUILDSETTINGS_API const TCHAR* GetBuildUser();
+
+	/**
+	 * Returns the domain name of the user which created these binaries
+	 */
+	BUILDSETTINGS_API const TCHAR* GetBuildUserDomain();
+
+	/**
+	 * Returns the name of the machine which created these binaries
+	 */
+	BUILDSETTINGS_API const TCHAR* GetBuildMachine();
 }
