[AutoRTFM] Enable AutoRTFM instrumentation in BuildSettings.

These are used by analytics, and analytics code can be triggered in a transactional context. For instance, Verse script can cause the game to end, and end-of-game can cause analytics data to be sent.

diff --git a/Engine/Source/Runtime/BuildSettings/BuildSettings.Build.cs b/Engine/Source/Runtime/BuildSettings/BuildSettings.Build.cs
--- a/Engine/Source/Runtime/BuildSettings/BuildSettings.Build.cs
+++ b/Engine/Source/Runtime/BuildSettings/BuildSettings.Build.cs
@@ -49,8 +49,6 @@ public class BuildSettings : ModuleRules
 
 		PrivateDefinitions.Add("SUPPRESS_PER_MODULE_INLINE_FILE"); // This module does not use core's standard operator new/delete overloads
 		
-		bDisableAutoRTFMInstrumentation = true;
-
 		if (Target.bWithLiveCoding && Target.LinkType == TargetLinkType.Monolithic)
 		{
 			PrivateDefinitions.Add(String.Format("UE_LIVE_CODING_ENGINE_DIR=\"{0}\"", UnrealBuildBase.Unreal.EngineDirectory.FullName.Replace("\\", "\\\\")));
