Added MemoryPools.PersistentAllocatorReserveSizeMB=2048 default ini value for Editor

diff --git a/Engine/Config/BaseEditor.ini b/Engine/Config/BaseEditor.ini
--- a/Engine/Config/BaseEditor.ini
+++ b/Engine/Config/BaseEditor.ini
@@ -713,3 +713,5 @@ bUseWidgetTemplateSelector=True
 +Filters=(FilterString="TexCoordRHIBuffer",DisplayName="Tex Coord Buffers")
 +Filters=(FilterString="FColorVertexBuffer",DisplayName="Color Vertex Buffers")
 
+[MemoryPools]
+PersistentAllocatorReserveSizeMB=2048
\ No newline at end of file
diff --git a/Engine/Source/Runtime/BuildSettings/BuildSettings.Build.cs b/Engine/Source/Runtime/BuildSettings/BuildSettings.Build.cs
--- a/Engine/Source/Runtime/BuildSettings/BuildSettings.Build.cs
+++ b/Engine/Source/Runtime/BuildSettings/BuildSettings.Build.cs
@@ -66,8 +66,10 @@ public class BuildSettings : ModuleRules
 		}
 	}
 
-	private ulong ReadPersistentAllocatorReserveSize(ConfigHierarchyType ConfigHierarchy)
+	private ulong GetPersistentAllocatorReserveSize()
 	{
+		// We have to separate persistent allocator reserve into Editor and Engine as some platforms, like iOS have very limited VM by default, so we can't reserve a lot of vm just in case
+		ConfigHierarchyType ConfigHierarchy = Target.Type == TargetType.Editor ? ConfigHierarchyType.Editor : ConfigHierarchyType.Engine;
 		ConfigHierarchy Ini = ConfigCache.ReadHierarchy(ConfigHierarchy, DirectoryReference.FromFile(Target.ProjectFile), Target.Platform);
 		int SizeOfPermanentObjectPool = 0;
 		if (Ini.GetInt32("/Script/Engine.GarbageCollectionSettings", "gc.SizeOfPermanentObjectPool", out SizeOfPermanentObjectPool))
@@ -78,17 +80,4 @@ public class BuildSettings : ModuleRules
 		Ini.GetInt32("MemoryPools", "PersistentAllocatorReserveSizeMB", out PersistentAllocatorReserveSizeMB);
 		return (ulong)PersistentAllocatorReserveSizeMB * 1024 * 1024;
 	}
-
-	private ulong GetPersistentAllocatorReserveSize()
-	{
-		ConfigHierarchyType ConfigHierarchy = Target.Type == TargetType.Editor ? ConfigHierarchyType.Editor : ConfigHierarchyType.Engine;
-		ulong PersistentAllocatorReserveSizeMB = ReadPersistentAllocatorReserveSize(ConfigHierarchy);
-		if (PersistentAllocatorReserveSizeMB == 0 && ConfigHierarchy == ConfigHierarchyType.Editor)
-		{
-			// if Editor doesn't define PersistentAllocatorReserveSizeMB, read the one from the Engine and simply double it
-			PersistentAllocatorReserveSizeMB = ReadPersistentAllocatorReserveSize(ConfigHierarchyType.Engine) * 2;
-		}
-
-		return PersistentAllocatorReserveSizeMB;
-	}
 }
