BuildSettings: Disallow loading PersistentAllocatorReserveSizeMB from a project ini when the environment is shared

#rnx

diff --git a/Engine/Source/Runtime/BuildSettings/BuildSettings.Build.cs b/Engine/Source/Runtime/BuildSettings/BuildSettings.Build.cs
--- a/Engine/Source/Runtime/BuildSettings/BuildSettings.Build.cs
+++ b/Engine/Source/Runtime/BuildSettings/BuildSettings.Build.cs
@@ -70,7 +70,7 @@ public class BuildSettings : ModuleRules
 	{
 		// We have to separate persistent allocator reserve into Editor and Engine as some platforms, like iOS have very limited VM by default, so we can't reserve a lot of vm just in case
 		ConfigHierarchyType ConfigHierarchy = Target.Type == TargetType.Editor ? ConfigHierarchyType.Editor : ConfigHierarchyType.Engine;
-		ConfigHierarchy Ini = ConfigCache.ReadHierarchy(ConfigHierarchy, DirectoryReference.FromFile(Target.ProjectFile), Target.Platform);
+		ConfigHierarchy Ini = ConfigCache.ReadHierarchy(ConfigHierarchy, Target.BuildEnvironment == TargetBuildEnvironment.Unique ? DirectoryReference.FromFile(Target.ProjectFile) : null, Target.Platform);
 		int SizeOfPermanentObjectPool = 0;
 		if (Ini.GetInt32("/Script/Engine.GarbageCollectionSettings", "gc.SizeOfPermanentObjectPool", out SizeOfPermanentObjectPool))
 		{
