UnrealBuildTool: Clean up some RHIShaderPlatform generated header logic

#rnx

diff --git a/Engine/Source/Programs/UnrealBuildTool/Configuration/UEBuildPlatform.cs b/Engine/Source/Programs/UnrealBuildTool/Configuration/UEBuildPlatform.cs
--- a/Engine/Source/Programs/UnrealBuildTool/Configuration/UEBuildPlatform.cs
+++ b/Engine/Source/Programs/UnrealBuildTool/Configuration/UEBuildPlatform.cs
@@ -1395,7 +1395,6 @@ namespace UnrealBuildTool
 		/// <param name="Receipt">Receipt for the target being deployed</param>
 		public abstract void Deploy(TargetReceipt Receipt);
 
-
 		/// <summary>
 		/// Struct used to add platform static shaders
 		/// </summary>
@@ -1405,6 +1404,7 @@ namespace UnrealBuildTool
 			/// Name of shader
 			/// </summary>
 			public string Name { set; get; }
+
 			/// <summary>
 			/// Index of shader
 			/// </summary>
@@ -1419,24 +1419,21 @@ namespace UnrealBuildTool
 		/// <param name="StaticShaders">Static shaders supported by current platform</param>
 		public void AddPlatformStaticShaders(ReadOnlyTargetRules Target, ModuleRules Rules, IEnumerable<StaticShader> StaticShaders)
 		{
-			if (Rules.Name == "BuildSettings") // Put list of platforms in BuildSettings since this array will vary depending on what platforms machine has installed
+			if (Target.bNeedsExtraShaderFormats || Target.Platform == Platform)
 			{
-				if (Target.bNeedsExtraShaderFormats || Target.Platform == Platform)
+				if (Rules.Name == "BuildSettings") // Put list of platforms in BuildSettings since this array will vary depending on what platforms machine has installed
 				{
-					List<string> content2 = (List<string>)Rules.FilesToGenerate["RHIStaticShaderPlatformNames.gen.cpp"];
-					content2.InsertRange(content2.Count - 2, StaticShaders.Select(shader =>
+					List<string> content = [.. Rules.FilesToGenerate["RHIStaticShaderPlatformNames.gen.cpp"]];
+					content.InsertRange(content.Count - 2, StaticShaders.Select(shader =>
 						$"	{{ TEXT(\"{shader.Name}\"), TEXT(\"{Platform}\"), {shader.Index} }},"));
+					Rules.FilesToGenerate["RHIStaticShaderPlatformNames.gen.cpp"] = [.. content];
 				}
-			}
-			else if (Rules.Name == "RHI") // Add one file per platform with EShaderPlatform entries
-			{
-				if (Target.bNeedsExtraShaderFormats || Target.Platform == Platform)
+				else if (Rules.Name == "RHI") // Add one file per platform with EShaderPlatform entries
 				{
-					List<string> content = ["#pragma once", "#include \"HAL/Platform.h\"", "enum EShaderPlatform : uint16;"];
-					foreach (StaticShader shader in StaticShaders)
-					{
-						content.Add($"inline constexpr EShaderPlatform SP_{shader.Name} = EShaderPlatform({shader.Index});");
-					}
+					IEnumerable<string> content = [
+						"#pragma once", "#include \"HAL/Platform.h\"", "enum EShaderPlatform : uint16;",
+						.. StaticShaders.Select(shader => $"inline constexpr EShaderPlatform SP_{shader.Name} = EShaderPlatform({shader.Index});")
+					];
 					Rules.FilesToGenerate.Add($"{Platform}RHIShaderPlatform.h", content);
 				}
 			}
diff --git a/Engine/Source/Runtime/BuildSettings/BuildSettings.Build.cs b/Engine/Source/Runtime/BuildSettings/BuildSettings.Build.cs
--- a/Engine/Source/Runtime/BuildSettings/BuildSettings.Build.cs
+++ b/Engine/Source/Runtime/BuildSettings/BuildSettings.Build.cs
@@ -1,11 +1,11 @@
 // Copyright Epic Games, Inc. All Rights Reserved.
 
 using EpicGames.Core;
-using UnrealBuildTool;
+using Microsoft.Extensions.Logging;
 using System;
 using System.Linq;
-using Microsoft.Extensions.Logging;
-using System.Collections.Generic;
+using UnrealBuildBase;
+using UnrealBuildTool;
 
 public class BuildSettings : ModuleRules
 {
@@ -18,13 +18,18 @@ public class BuildSettings : ModuleRules
 
 		bRequiresImplementModule = false;
 
+		if (Target.bBuildEditor)
+		{
+			bDisableAutoRTFMInstrumentation = true;
+		}
+
 		PrivateDefinitions.Add($"ENGINE_VERSION_MAJOR={Target.Version.MajorVersion}");
 		PrivateDefinitions.Add($"ENGINE_VERSION_MINOR={Target.Version.MinorVersion}");
 		PrivateDefinitions.Add($"ENGINE_VERSION_HOTFIX={Target.Version.PatchVersion}");
 		PrivateDefinitions.Add($"ENGINE_VERSION_STRING=\"{Target.Version.MajorVersion}.{Target.Version.MinorVersion}.{Target.Version.PatchVersion}-{Target.BuildVersion}\"");
 		PrivateDefinitions.Add($"ENGINE_IS_LICENSEE_VERSION={(Target.Version.IsLicenseeVersion ? "true" : "false")}");
 		PrivateDefinitions.Add($"ENGINE_IS_PROMOTED_BUILD={(Target.Version.IsPromotedBuild ? "true" : "false")}");
-		
+
 		if (!Target.GlobalDefinitions.Any(x => x.Contains("CURRENT_CHANGELIST", StringComparison.Ordinal)))
 		{
 			PrivateDefinitions.Add($"CURRENT_CHANGELIST={Target.Version.Changelist}");
@@ -35,53 +40,38 @@ public class BuildSettings : ModuleRules
 		PrivateDefinitions.Add($"BUILD_VERSION=\"{Target.BuildVersion}\"");
 		PrivateDefinitions.Add($"BUILD_SOURCE_URL=\"{Target.Version.BuildURL}\"");
 
-		string userName = string.Empty;
-		string userDomainName = string.Empty;
-		string machineName = string.Empty;
-		if (Target.bEnablePrivateBuildInformation)
-		{
-			userName = Environment.UserName;
-			userDomainName = Environment.UserDomainName;
-			machineName = UnrealBuildBase.Unreal.MachineName;
-		}
-		PrivateDefinitions.Add($"BUILD_USER=\"{userName}\"");
-		PrivateDefinitions.Add($"BUILD_USERDOMAINNAME=\"{userDomainName}\"");
-		PrivateDefinitions.Add($"BUILD_MACHINENAME=\"{machineName}\"");
+		PrivateDefinitions.Add($"BUILD_USER=\"{(Target.bEnablePrivateBuildInformation ? Environment.UserName : String.Empty)}\"");
+		PrivateDefinitions.Add($"BUILD_USERDOMAINNAME=\"{(Target.bEnablePrivateBuildInformation ? Environment.UserDomainName : String.Empty)}\"");
+		PrivateDefinitions.Add($"BUILD_MACHINENAME=\"{(Target.bEnablePrivateBuildInformation ? Unreal.MachineName : String.Empty)}\"");
 
 		PrivateDefinitions.Add("SUPPRESS_PER_MODULE_INLINE_FILE"); // This module does not use core's standard operator new/delete overloads
-		
+
 		if (Target.bWithLiveCoding && Target.LinkType == TargetLinkType.Monolithic)
 		{
-			PrivateDefinitions.Add(String.Format("UE_LIVE_CODING_ENGINE_DIR=\"{0}\"", UnrealBuildBase.Unreal.EngineDirectory.FullName.Replace("\\", "\\\\")));
+			PrivateDefinitions.Add($"UE_LIVE_CODING_ENGINE_DIR=\"{Unreal.EngineDirectory.FullName.Replace("\\", "\\\\")}\"");
 			if (Target.ProjectFile != null)
 			{
-				PrivateDefinitions.Add(String.Format("UE_LIVE_CODING_PROJECT=\"{0}\"", Target.ProjectFile.FullName.Replace("\\", "\\\\")));
+				PrivateDefinitions.Add($"UE_LIVE_CODING_PROJECT=\"{Target.ProjectFile.FullName.Replace("\\", "\\\\")}\"");
 			}
 		}
 
-		PrivateDefinitions.Add("UE_PERSISTENT_ALLOCATOR_RESERVE_SIZE=" + GetPersistentAllocatorReserveSize().ToString() + "ULL");
-
-		if (Target.bBuildEditor)
-		{
-			bDisableAutoRTFMInstrumentation = true;
-		}
+		PrivateDefinitions.Add($"UE_PERSISTENT_ALLOCATOR_RESERVE_SIZE={GetPersistentAllocatorReserveSize()}ULL");
 
 		if (!FilesToGenerate.ContainsKey("RHIStaticShaderPlatformNames.gen.h"))
 		{
-			FilesToGenerate.TryAdd("RHIStaticShaderPlatformNames.gen.h",
-				[
-					"#pragma once",
-							"#include \"HAL/Platform.h\"",
-							"struct FStaticNameMapEntry { const TCHAR* Name; const TCHAR* Platform; const int Enum; };",
-							"extern BUILDSETTINGS_API FStaticNameMapEntry GStaticShaderNames[];"
-				]);
-			FilesToGenerate.Add("RHIStaticShaderPlatformNames.gen.cpp", new List<string>([
+			FilesToGenerate.TryAdd("RHIStaticShaderPlatformNames.gen.h", [
+				"#pragma once",
+				"#include \"HAL/Platform.h\"",
+				"struct FStaticNameMapEntry { const TCHAR* Name; const TCHAR* Platform; const int Enum; };",
+				"extern BUILDSETTINGS_API FStaticNameMapEntry GStaticShaderNames[];"
+			]);
+			FilesToGenerate.Add("RHIStaticShaderPlatformNames.gen.cpp", [
 				"#include \"RHIStaticShaderPlatformNames.gen.h\"",
-						"BUILDSETTINGS_API FStaticNameMapEntry GStaticShaderNames[] = ",
-						"{",
-						"	{ nullptr, nullptr, 0 },",
-						"};"
-				]));
+				"BUILDSETTINGS_API FStaticNameMapEntry GStaticShaderNames[] = ",
+				"{",
+				"	{ nullptr, nullptr, 0 },",
+				"};"
+			]);
 		}
 	}
 
@@ -90,13 +80,11 @@ public class BuildSettings : ModuleRules
 		// We have to separate persistent allocator reserve into Editor and Engine as some platforms, like iOS have very limited VM by default, so we can't reserve a lot of vm just in case
 		ConfigHierarchyType ConfigHierarchy = Target.Type == TargetType.Editor ? ConfigHierarchyType.Editor : ConfigHierarchyType.Engine;
 		ConfigHierarchy Ini = ConfigCache.ReadHierarchy(ConfigHierarchy, Target.BuildEnvironment == TargetBuildEnvironment.Unique ? DirectoryReference.FromFile(Target.ProjectFile) : null, Target.Platform);
-		int SizeOfPermanentObjectPool = 0;
-		if (Ini.GetInt32("/Script/Engine.GarbageCollectionSettings", "gc.SizeOfPermanentObjectPool", out SizeOfPermanentObjectPool))
+		if (Ini.GetInt32("/Script/Engine.GarbageCollectionSettings", "gc.SizeOfPermanentObjectPool", out int _))
 		{
-			Target.Logger.LogWarning("/Script/Engine.GarbageCollectionSettings, gc.SizeOfPermanentObjectPool ini for Project {0} setting was deprecated in a favor of MemoryPools, PersistentAllocatorReserveSizeMB", Target.ProjectFile.ToString());
+			Target.Logger.LogWarning("/Script/Engine.GarbageCollectionSettings, gc.SizeOfPermanentObjectPool ini for Project {Project} setting was deprecated in a favor of MemoryPools, PersistentAllocatorReserveSizeMB", Target.ProjectFile);
 		}
-		int PersistentAllocatorReserveSizeMB = 0;
-		Ini.GetInt32("MemoryPools", "PersistentAllocatorReserveSizeMB", out PersistentAllocatorReserveSizeMB);
+		Ini.GetInt32("MemoryPools", "PersistentAllocatorReserveSizeMB", out int PersistentAllocatorReserveSizeMB);
 		return (ulong)PersistentAllocatorReserveSizeMB * 1024 * 1024;
 	}
 }
