[UBT]
* Added support for adding variables to buildsettings from target rules. This can be used to avoid GlobalDefinitions (which will cause all actions in the build to be different). BuildSettings variables can be accessed using BuildSettings::GetVariable(name);

diff --git a/Engine/Source/Programs/UnrealBuildTool/Configuration/ReadOnlyTargetRules.cs b/Engine/Source/Programs/UnrealBuildTool/Configuration/ReadOnlyTargetRules.cs
--- a/Engine/Source/Programs/UnrealBuildTool/Configuration/ReadOnlyTargetRules.cs
+++ b/Engine/Source/Programs/UnrealBuildTool/Configuration/ReadOnlyTargetRules.cs
@@ -539,6 +539,8 @@ namespace UnrealBuildTool
 
 		public IReadOnlyList<string> ProjectDefinitions => Inner.ProjectDefinitions.AsReadOnly();
 
+		public IReadOnlyList<(string, string)> BuildSettingsVariables => Inner.BuildSettingsVariables.AsReadOnly();
+
 		public ReadOnlyAdditionalProperties AdditionalProperties => Inner.AdditionalProperties;
 
 		public string? LaunchModuleName => Inner.LaunchModuleName;
diff --git a/Engine/Source/Programs/UnrealBuildTool/Configuration/TargetRules.cs b/Engine/Source/Programs/UnrealBuildTool/Configuration/TargetRules.cs
--- a/Engine/Source/Programs/UnrealBuildTool/Configuration/TargetRules.cs
+++ b/Engine/Source/Programs/UnrealBuildTool/Configuration/TargetRules.cs
@@ -2645,6 +2645,11 @@ namespace UnrealBuildTool
 		[CommandLine("-ProjectDefine:")]
 		public List<string> ProjectDefinitions { get; } = [];
 
+		/// <summary>
+		/// Variables that will be stored in BuildSettings module. Can be fetched with BuildSettings::GetVariable
+		/// </summary>
+		public List<(string key, string value)> BuildSettingsVariables { get; } = [];
+
 		/// <summary>
 		/// Environment variables. Should be used instead of Environment.Set/GetEnvironmentVariable in .build.cs
 		/// </summary>
diff --git a/Engine/Source/Programs/UnrealBuildTool/Configuration/UEBuildModuleCPP.cs b/Engine/Source/Programs/UnrealBuildTool/Configuration/UEBuildModuleCPP.cs
--- a/Engine/Source/Programs/UnrealBuildTool/Configuration/UEBuildModuleCPP.cs
+++ b/Engine/Source/Programs/UnrealBuildTool/Configuration/UEBuildModuleCPP.cs
@@ -2646,6 +2646,21 @@ namespace UnrealBuildTool
 					sb.Append(local.FullName.Replace("\\", "\\\\")).Append(';');
 				}
 				Result.Definitions.Add($"UE_VFS_PATHS=\"{sb}\"");
+
+				sb.Clear();
+				foreach ((string key, string value) in Target.BuildSettingsVariables)
+				{
+					if (key.Contains(';') || value.Contains(';'))
+					{
+						throw new Exception("BuildSettingsDefinitions can't contain keys or values with character ';'");
+					}
+
+					sb.Append(key).Append(';').Append(value).Append(';'); ;
+				}
+				if (sb.Length != 0)
+				{
+					Result.Definitions.Add($"UE_BUILDSETTINGS_VARIABLES=\"{sb}\"");
+				}
 			}
 
 			// Add the project definitions
diff --git a/Engine/Source/Runtime/BuildSettings/Private/BuildSettings.cpp b/Engine/Source/Runtime/BuildSettings/Private/BuildSettings.cpp
--- a/Engine/Source/Runtime/BuildSettings/Private/BuildSettings.cpp
+++ b/Engine/Source/Runtime/BuildSettings/Private/BuildSettings.cpp
@@ -1,6 +1,8 @@
 // Copyright Epic Games, Inc. All Rights Reserved.
 
 #include "BuildSettings.h"
+#include "Containers/StringView.h"
+#include "Misc/CString.h"
 
 namespace BuildSettings
 {
@@ -121,4 +123,41 @@ namespace BuildSettings
 	{
 		return TEXT(UE_VFS_PATHS);
 	}
+
+	bool GetVariable(const TCHAR* InName, TCHAR* OutValue, int OutValueCapacity)
+	{
+		#ifdef UE_BUILDSETTINGS_VARIABLES
+		const TCHAR* Name = TEXT(UE_BUILDSETTINGS_VARIABLES);
+		while (true)
+		{
+			const TCHAR* NameEnd = FCString::Strchr(Name, ';');
+			if (!NameEnd)
+			{
+				return false;
+			}
+
+			const TCHAR* Value = NameEnd + 1;
+			const TCHAR* ValueEnd = FCString::Strchr(Value, ';');
+
+			FStringView NameView(Name, int32(NameEnd - Name));
+			if (NameView.Equals(InName))
+			{
+				if (OutValue)
+				{
+					FStringView ValueView(Value, int32(ValueEnd - Value));
+
+					checkf(OutValueCapacity > ValueView.Len(), TEXT("Capacity of string sent in to BuildSettings::GetVariable is too short"));
+
+					ValueView.CopyString(OutValue, OutValueCapacity, 0);
+				}
+				return true;
+			}
+
+			Name = ValueEnd + 1;
+
+		}
+		#endif
+
+		return false;
+	}
 }
diff --git a/Engine/Source/Runtime/BuildSettings/Public/BuildSettings.h b/Engine/Source/Runtime/BuildSettings/Public/BuildSettings.h
--- a/Engine/Source/Runtime/BuildSettings/Public/BuildSettings.h
+++ b/Engine/Source/Runtime/BuildSettings/Public/BuildSettings.h
@@ -147,4 +147,9 @@ namespace BuildSettings
 	 * Semicolon separated string containing virtual;local;virtual;local; paths. Ends with semicolon
 	 */
 	BUILDSETTINGS_API const TCHAR* GetVfsPathsWide();
+
+	/**
+	 * Returns true if variable exists. Writes value into OutValue if not null
+	 */
+	BUILDSETTINGS_API bool GetVariable(const TCHAR* InName, TCHAR* OutValue, int OutValueCapacity);
 }
