[BuildSettings]
* Fixed bug in BuildSettings::GetVariable where OutValue string would not be null terminated

diff --git a/Engine/Source/Runtime/BuildSettings/Private/BuildSettings.cpp b/Engine/Source/Runtime/BuildSettings/Private/BuildSettings.cpp
--- a/Engine/Source/Runtime/BuildSettings/Private/BuildSettings.cpp
+++ b/Engine/Source/Runtime/BuildSettings/Private/BuildSettings.cpp
@@ -148,7 +148,8 @@ namespace BuildSettings
 
 					checkf(OutValueCapacity > ValueView.Len(), TEXT("Capacity of string sent in to BuildSettings::GetVariable is too short"));
 
-					ValueView.CopyString(OutValue, OutValueCapacity, 0);
+					int32 CopyCount = ValueView.CopyString(OutValue, OutValueCapacity, 0);
+					OutValue[CopyCount] = 0;
 				}
 				return true;
 			}
